GitHub 1205 : Fix UProceduralMeshComponent Not Updating Bounds #1205
Moved BodySetup include

diff --git a/Engine/Source/Runtime/ProceduralMeshComponent/Private/ProceduralMeshComponent.cpp b/Engine/Source/Runtime/ProceduralMeshComponent/Private/ProceduralMeshComponent.cpp
--- a/Engine/Source/Runtime/ProceduralMeshComponent/Private/ProceduralMeshComponent.cpp
+++ b/Engine/Source/Runtime/ProceduralMeshComponent/Private/ProceduralMeshComponent.cpp
@@ -3,7 +3,7 @@
 #include "ProceduralMeshComponentPrivatePCH.h"
 #include "ProceduralMeshComponent.h"
 #include "DynamicMeshBuilder.h"
-
+#include "PhysicsEngine/BodySetup.h"
 
 /** Resource array to pass  */
 class FProcMeshVertexResourceArray : public FResourceArrayInterface
@@ -379,6 +379,9 @@ void UProceduralMeshComponent::UpdateLocalBounds()
 	}
 
 	LocalBounds = LocalBox.IsValid ? FBoxSphereBounds(LocalBox) : FBoxSphereBounds(FVector(0, 0, 0), FVector(0, 0, 0), 0); // fallback to reset box sphere bounds
+
+	// Update global bounds
+	UpdateBounds();
 }
 
 FPrimitiveSceneProxy* UProceduralMeshComponent::CreateSceneProxy()
