GitHub 1205 : Fix UProceduralMeshComponent Not Updating Bounds #1205
Moved BodySetup include

diff --git a/Engine/Source/Runtime/ProceduralMeshComponent/Private/ProceduralMeshComponent.cpp b/Engine/Source/Runtime/ProceduralMeshComponent/Private/ProceduralMeshComponent.cpp
--- a/Engine/Source/Runtime/ProceduralMeshComponent/Private/ProceduralMeshComponent.cpp
+++ b/Engine/Source/Runtime/ProceduralMeshComponent/Private/ProceduralMeshComponent.cpp
@@ -3,7 +3,7 @@
 #include "ProceduralMeshComponentPrivatePCH.h"
 #include "ProceduralMeshComponent.h"
 #include "DynamicMeshBuilder.h"
-#include "PhysicsEngine/BodySetup.h"
+
 
 /** Resource array to pass  */
 class FProcMeshVertexResourceArray : public FResourceArrayInterface
@@ -379,9 +379,6 @@ void UProceduralMeshComponent::UpdateLocalBounds()
 	}
 
 	LocalBounds = LocalBox.IsValid ? FBoxSphereBounds(LocalBox) : FBoxSphereBounds(FVector(0, 0, 0), FVector(0, 0, 0), 0); // fallback to reset box sphere bounds
-
-	// Update global bounds
-	UpdateBounds();
 }
 
 FPrimitiveSceneProxy* UProceduralMeshComponent::CreateSceneProxy()
diff --git a/Engine/Source/Runtime/ProceduralMeshComponent/Private/ProceduralMeshComponentPrivatePCH.h b/Engine/Source/Runtime/ProceduralMeshComponent/Private/ProceduralMeshComponentPrivatePCH.h
--- a/Engine/Source/Runtime/ProceduralMeshComponent/Private/ProceduralMeshComponentPrivatePCH.h
+++ b/Engine/Source/Runtime/ProceduralMeshComponent/Private/ProceduralMeshComponentPrivatePCH.h
@@ -2,7 +2,6 @@
 
 #include "CoreUObject.h"
 #include "Engine.h"
-#include "PhysicsEngine/BodySetup.h"
 
 // You should place include statements to your module's private header files here.  You only need to
 // add includes for headers that are used in most of your module's source files though.
