- Add 'fast path' UpdateMeshSection function to ProcMeshComp (updates collision as well)
- Add SetMeshSectionVisible/IsMeshSectionVisible functions to PMC
- Fix crash in ProcMeshComp when Triangles array was not a multiple of 3

diff --git a/Engine/Source/Runtime/Engine/Classes/PhysicsEngine/BodyInstance.h b/Engine/Source/Runtime/Engine/Classes/PhysicsEngine/BodyInstance.h
--- a/Engine/Source/Runtime/Engine/Classes/PhysicsEngine/BodyInstance.h
+++ b/Engine/Source/Runtime/Engine/Classes/PhysicsEngine/BodyInstance.h
@@ -625,6 +625,7 @@ public:
 	 */
 	bool UpdateBodyScale(const FVector& InScale3D);
 
+	/** Dynamically update the vertices of per-poly collision for this body. */
 	void UpdateTriMeshVertices(const TArray<FVector> & NewPositions);
 
 	/** Returns the center of mass of this body (in world space) */
diff --git a/Engine/Source/Runtime/ProceduralMeshComponent/Private/ProceduralMeshComponent.cpp b/Engine/Source/Runtime/ProceduralMeshComponent/Private/ProceduralMeshComponent.cpp
--- a/Engine/Source/Runtime/ProceduralMeshComponent/Private/ProceduralMeshComponent.cpp
+++ b/Engine/Source/Runtime/ProceduralMeshComponent/Private/ProceduralMeshComponent.cpp
@@ -121,12 +121,38 @@ public:
 	FProcMeshIndexBuffer IndexBuffer;
 	/** Vertex factory for this section */
 	FProcMeshVertexFactory VertexFactory;
+	/** Whether this section is currently visible */
+	bool bSectionVisible;
 
 	FProcMeshProxySection()
-		: Material(NULL)
+	: Material(NULL)
+	, bSectionVisible(true)
 	{}
 };
 
+/** 
+ *	Struct used to send update to mesh data 
+ *	Arrays may be empty, in which case no update is performed.
+ */
+class FProcMeshSectionUpdateData
+{
+public:
+	/** Section to update */
+	int32 TargetSection;
+	/** New vertex information */
+	TArray<FProcMeshVertex> NewVertexBuffer;
+};
+
+static void ConvertProcMeshToDynMeshVertex(FDynamicMeshVertex& Vert, const FProcMeshVertex& ProcVert)
+{
+	Vert.Position = ProcVert.Position;
+	Vert.Color = ProcVert.Color;
+	Vert.TextureCoordinate = ProcVert.UV0;
+	Vert.TangentX = ProcVert.Tangent.TangentX;
+	Vert.TangentZ = ProcVert.Normal;
+	Vert.TangentZ.Vector.W = ProcVert.Tangent.bFlipTangentY ? 0 : 255;
+}
+
 /** Procedural mesh scene proxy */
 class FProceduralMeshSceneProxy : public FPrimitiveSceneProxy
 {
@@ -154,15 +180,9 @@ public:
 				// Copy verts
 				for (int VertIdx = 0; VertIdx < NumVerts; VertIdx++)
 				{
-					FProcMeshVertex& ProcVert = SrcSection.ProcVertexBuffer[VertIdx];
-
+					const FProcMeshVertex& ProcVert = SrcSection.ProcVertexBuffer[VertIdx];
 					FDynamicMeshVertex& Vert = NewSection->VertexBuffer.Vertices[VertIdx];
-					Vert.Position = ProcVert.Position;
-					Vert.Color = ProcVert.Color;
-					Vert.TextureCoordinate = ProcVert.UV0;
-					Vert.TangentX = ProcVert.Tangent.TangentX;
-					Vert.TangentZ = ProcVert.Normal;
-					Vert.TangentZ.Vector.W = ProcVert.Tangent.bFlipTangentY ? 0 : 255;
+					ConvertProcMeshToDynMeshVertex(Vert, ProcVert);
 				}
 
 				// Copy index buffer
@@ -183,6 +203,9 @@ public:
 					NewSection->Material = UMaterial::GetDefaultMaterial(MD_Surface);
 				}
 
+				// Copy visibility info
+				NewSection->bSectionVisible = SrcSection.bSectionVisible;
+
 				// Save ref to new section
 				Sections[SectionIdx] = NewSection;
 			}
@@ -203,6 +226,52 @@ public:
 		}
 	}
 
+	/** Called on render thread to assign new dynamic data */
+	void UpdateSection_RenderThread(FProcMeshSectionUpdateData* SectionData)
+	{
+		check(IsInRenderingThread());
+
+		// Check we have data 
+		if(	SectionData != nullptr) 			
+		{
+			// Check it references a valid section
+			if (SectionData->TargetSection < Sections.Num() &&
+				Sections[SectionData->TargetSection] != nullptr)
+			{
+				FProcMeshProxySection* Section = Sections[SectionData->TargetSection];
+
+				// Lock vertex buffer
+				const int32 NumVerts = SectionData->NewVertexBuffer.Num();
+				FDynamicMeshVertex* VertexBufferData = (FDynamicMeshVertex*)RHILockVertexBuffer(Section->VertexBuffer.VertexBufferRHI, 0, NumVerts * sizeof(FDynamicMeshVertex), RLM_WriteOnly);
+			
+				// Iterate through vertex data, copying in new info
+				for(int32 VertIdx=0; VertIdx<NumVerts; VertIdx++)
+				{
+					const FProcMeshVertex& ProcVert = SectionData->NewVertexBuffer[VertIdx];
+					FDynamicMeshVertex& Vert = VertexBufferData[VertIdx];
+					ConvertProcMeshToDynMeshVertex(Vert, ProcVert);
+				}
+
+				// Unlock vertex buffer
+				RHIUnlockVertexBuffer(Section->VertexBuffer.VertexBufferRHI);
+			}
+
+			// Free data sent from game thread
+			delete SectionData;
+		}
+	}
+
+	void SetSectionVisibility_RenderThread(int32 SectionIndex, bool bNewVisibility)
+	{
+		check(IsInRenderingThread());
+
+		if(	SectionIndex < Sections.Num() &&
+			Sections[SectionIndex] != nullptr)
+		{
+			Sections[SectionIndex]->bSectionVisible = bNewVisibility;
+		}
+	}
+
 	virtual void GetDynamicMeshElements(const TArray<const FSceneView*>& Views, const FSceneViewFamily& ViewFamily, uint32 VisibilityMap, FMeshElementCollector& Collector) const override
 	{
 		// Set up wireframe material (if needed)
@@ -222,7 +291,7 @@ public:
 		// Iterate over sections
 		for (const FProcMeshProxySection* Section : Sections)
 		{
-			if (Section != nullptr)
+			if (Section != nullptr && Section->bSectionVisible)
 			{
 				FMaterialRenderProxy* MaterialProxy = bWireframe ? WireframeMaterialInstance : Section->Material->GetRenderProxy(IsSelected());
 
@@ -337,7 +406,9 @@ void UProceduralMeshComponent::CreateMeshSection(int32 SectionIndex, const TArra
 	}
 
 	// Copy index buffer (clamping to vertex range)
-	const int32 NumTriIndices = Triangles.Num();
+	int32 NumTriIndices = Triangles.Num();
+	NumTriIndices = (NumTriIndices/3) * 3; // Ensure we have exact number of triangles (array is multiple of 3 long)
+
 	NewSection.ProcIndexBuffer.Reset();
 	NewSection.ProcIndexBuffer.AddUninitialized(NumTriIndices);
 	for (int32 IndexIdx = 0; IndexIdx < NumTriIndices; IndexIdx++)
@@ -349,7 +420,105 @@ void UProceduralMeshComponent::CreateMeshSection(int32 SectionIndex, const TArra
 
 	UpdateLocalBounds(); // Update overall bounds
 	UpdateCollision(); // Mark collision as dirty
-	MarkRenderStateDirty(); // Recreate scene proxy
+	MarkRenderStateDirty(); // New section requires recreating scene proxy
+}
+
+void UProceduralMeshComponent::UpdateMeshSection(int32 SectionIndex, const TArray<FVector>& Vertices, const TArray<FVector>& Normals, const TArray<FVector2D>& UV0, const TArray<FColor>& VertexColors, const TArray<FProcMeshTangent>& Tangents)
+{
+	if(SectionIndex < ProcMeshSections.Num())
+	{
+		FProcMeshSection& Section = ProcMeshSections[SectionIndex];
+		const int32 NumVerts = Section.ProcVertexBuffer.Num();
+
+		// See if positions are changing
+		const bool bPositionsChanging = (Vertices.Num() == NumVerts);
+
+		// Update bounds, if we are getting new position data
+		if (bPositionsChanging)
+		{
+			Section.SectionLocalBox.Init();
+		}
+
+		// Iterate through vertex data, copying in new info
+		for (int32 VertIdx = 0; VertIdx < NumVerts; VertIdx++)
+		{
+			FProcMeshVertex& ModifyVert = Section.ProcVertexBuffer[VertIdx];
+
+			// Position data
+			if (Vertices.Num() == NumVerts)
+			{
+				ModifyVert.Position = Vertices[VertIdx];
+				Section.SectionLocalBox += ModifyVert.Position;
+			}
+
+			// Normal data
+			if (Normals.Num() == NumVerts)
+			{
+				ModifyVert.Normal = Normals[VertIdx];
+			}
+
+			// Tangent data 
+			if (Tangents.Num() == NumVerts)
+			{
+				ModifyVert.Tangent = Tangents[VertIdx];
+			}
+
+			// UV data
+			if (UV0.Num() == NumVerts)
+			{
+				ModifyVert.UV0 = UV0[VertIdx];
+			}
+
+			// Color data
+			if (VertexColors.Num() == NumVerts)
+			{
+				ModifyVert.Color = VertexColors[VertIdx];
+			}
+		}
+
+		// Create data to update section
+		FProcMeshSectionUpdateData* SectionData = new FProcMeshSectionUpdateData;
+		SectionData->TargetSection = SectionIndex;
+		SectionData->NewVertexBuffer = Section.ProcVertexBuffer;
+
+		// Enqueue command to send to render thread
+		ENQUEUE_UNIQUE_RENDER_COMMAND_TWOPARAMETER(
+			FProcMeshSectionUpdate,
+			FProceduralMeshSceneProxy*, ProcMeshSceneProxy, (FProceduralMeshSceneProxy*)SceneProxy,
+			FProcMeshSectionUpdateData*, SectionData, SectionData,
+			{
+				ProcMeshSceneProxy->UpdateSection_RenderThread(SectionData);
+			}
+		);
+
+		// If we have collision enabled on this section, update that too
+		if(bPositionsChanging && Section.bEnableCollision)
+		{
+			TArray<FVector> CollisionPositions;
+
+			// We have one collision mesh for all sections, so need to build array of _all_ positions
+			for (const FProcMeshSection& CollisionSection : ProcMeshSections)
+			{
+				// If section has collision, copy it
+				if (CollisionSection.bEnableCollision)
+				{
+					for (int32 VertIdx = 0; VertIdx < Section.ProcVertexBuffer.Num(); VertIdx++)
+					{
+						CollisionPositions.Add(Section.ProcVertexBuffer[VertIdx].Position);
+					}
+				}
+			}
+
+			// Pass new positions to trimesh
+			BodyInstance.UpdateTriMeshVertices(CollisionPositions);
+		}
+
+		if (bPositionsChanging)
+		{
+			UpdateLocalBounds(); // Update overall bounds
+			MarkRenderTransformDirty(); // Need to send new bounds to render thread
+		}
+	}
 }
 
 void UProceduralMeshComponent::ClearMeshSection(int32 SectionIndex)
@@ -371,6 +540,32 @@ void UProceduralMeshComponent::ClearAllMeshSections()
 	MarkRenderStateDirty();
 }
 
+void UProceduralMeshComponent::SetMeshSectionVisible(int32 SectionIndex, bool bNewVisibility)
+{
+	if(SectionIndex < ProcMeshSections.Num())
+	{
+		// Set game thread state
+		ProcMeshSections[SectionIndex].bSectionVisible = bNewVisibility;
+
+		// Enqueue command to modify render thread info
+		ENQUEUE_UNIQUE_RENDER_COMMAND_THREEPARAMETER(
+			FProcMeshSectionVisibilityUpdate,
+			FProceduralMeshSceneProxy*, ProcMeshSceneProxy, (FProceduralMeshSceneProxy*)SceneProxy,
+			int32, SectionIndex, SectionIndex,
+			bool, bNewVisibility, bNewVisibility,
+			{
+				ProcMeshSceneProxy->SetSectionVisibility_RenderThread(SectionIndex, bNewVisibility);
+			}
+		);
+	}
+}
+
+bool UProceduralMeshComponent::IsMeshSectionVisible(int32 SectionIndex) const
+{
+	return (SectionIndex < ProcMeshSections.Num()) ? ProcMeshSections[SectionIndex].bSectionVisible : false;
+}
+
+
 void UProceduralMeshComponent::UpdateLocalBounds()
 {
 	FBox LocalBox(0);
diff --git a/Engine/Source/Runtime/ProceduralMeshComponent/Public/ProceduralMeshComponent.h b/Engine/Source/Runtime/ProceduralMeshComponent/Public/ProceduralMeshComponent.h
--- a/Engine/Source/Runtime/ProceduralMeshComponent/Public/ProceduralMeshComponent.h
+++ b/Engine/Source/Runtime/ProceduralMeshComponent/Public/ProceduralMeshComponent.h
@@ -13,13 +13,13 @@ struct FProcMeshTangent
 {
 	GENERATED_USTRUCT_BODY()
 
-		/** Direction of X tangent for this vertex */
-		UPROPERTY(EditAnywhere, BlueprintReadWrite, Category = Tangent)
-		FVector TangentX;
+	/** Direction of X tangent for this vertex */
+	UPROPERTY(EditAnywhere, BlueprintReadWrite, Category = Tangent)
+	FVector TangentX;
 
 	/** Bool that indicates whether we should flip the Y tangent when we compute it using cross product */
 	UPROPERTY(EditAnywhere, BlueprintReadWrite, Category = Tangent)
-		bool bFlipTangentY;
+	bool bFlipTangentY;
 
 	FProcMeshTangent()
 		: TangentX(1.f, 0.f, 0.f)
@@ -43,20 +43,25 @@ struct FProcMeshVertex
 {
 	GENERATED_USTRUCT_BODY()
 
-		UPROPERTY(EditAnywhere, BlueprintReadWrite, Category = Vertex)
-		FVector Position;
+	/** Vertex position */
+	UPROPERTY(EditAnywhere, BlueprintReadWrite, Category = Vertex)
+	FVector Position;
 
+	/** Vertex normal */
 	UPROPERTY(EditAnywhere, BlueprintReadWrite, Category = Vertex)
-		FVector Normal;
+	FVector Normal;
 
+	/** Vertex tangent */
 	UPROPERTY(EditAnywhere, BlueprintReadWrite, Category = Vertex)
-		FProcMeshTangent Tangent;
+	FProcMeshTangent Tangent;
 
+	/** Vertex color */
 	UPROPERTY(EditAnywhere, BlueprintReadWrite, Category = Vertex)
-		FColor Color;
+	FColor Color;
 
+	/** Vertex texture co-ordinate */
 	UPROPERTY(EditAnywhere, BlueprintReadWrite, Category = Vertex)
-		FVector2D UV0;
+	FVector2D UV0;
 
 
 	FProcMeshVertex()
@@ -74,22 +79,29 @@ struct FProcMeshSection
 {
 	GENERATED_USTRUCT_BODY()
 
-		/** Vertex buffer for this section */
-		UPROPERTY()
-		TArray<FProcMeshVertex> ProcVertexBuffer;
+	/** Vertex buffer for this section */
+	UPROPERTY()
+	TArray<FProcMeshVertex> ProcVertexBuffer;
+
 	/** Index buffer for this section */
 	UPROPERTY()
-		TArray<int32> ProcIndexBuffer;
+	TArray<int32> ProcIndexBuffer;
 	/** Local bounding box of section */
 	UPROPERTY()
-		FBox SectionLocalBox;
+	FBox SectionLocalBox;
+
 	/** Should we build collision data for triangles in this section */
 	UPROPERTY()
-		bool bEnableCollision;
+	bool bEnableCollision;
+
+	/** Should we display this section */
+	UPROPERTY()
+	bool bSectionVisible;
 
 	FProcMeshSection()
 		: SectionLocalBox(0)
 		, bEnableCollision(false)
+		, bSectionVisible(true)
 	{}
 
 	/** Reset this section, clear all mesh info. */
@@ -99,6 +111,7 @@ struct FProcMeshSection
 		ProcIndexBuffer.Empty();
 		SectionLocalBox.Init();
 		bEnableCollision = false;
+		bSectionVisible = true;
 	}
 };
 
@@ -111,28 +124,46 @@ class PROCEDURALMESHCOMPONENT_API UProceduralMeshComponent : public UMeshCompone
 {
 	GENERATED_UCLASS_BODY()
 
-		/**
-		*	Create/replace a section for this procedural mesh component.
-		*	@param	SectionIndex		Index of the section to create or replace.
-		*	@param	Vertices			Vertex buffer of all vertex positions to use for this mesh section.
-		*	@param	Triangles			Index buffer indicating which vertices make up each triangle. Length must be a multiple of 3.
-		*	@param	Normals				Optional array of normal vectors for each vertex. If supplied, must be same length as Vertices array.
-		*	@param	UV0					Optional array of texture co-ordinates for each vertex. If supplied, must be same length as Vertices array.
-		*	@param	VertexColors		Optional array of colors for each vertex. If supplied, must be same length as Vertices array.
-		*	@param	Tangents			Optional array of tangent vector for each vertex. If supplied, must be same length as Vertices array.
-		*	@param	bCreateCollision	Indicates whether collision should be created for this section. This adds significant cost.
-		*/
-		UFUNCTION(BlueprintCallable, Category = "Components|ProceduralMesh", meta = (AutoCreateRefTerm = "Normals,UV0,VertexColors,Tangents"))
-		void CreateMeshSection(int32 SectionIndex, const TArray<FVector>& Vertices, const TArray<int32>& Triangles, const TArray<FVector>& Normals, const TArray<FVector2D>& UV0, const TArray<FColor>& VertexColors, const TArray<FProcMeshTangent>& Tangents, bool bCreateCollision);
+	/**
+	 *	Create/replace a section for this procedural mesh component.
+	 *	@param	SectionIndex		Index of the section to create or replace.
+	 *	@param	Vertices			Vertex buffer of all vertex positions to use for this mesh section.
+	 *	@param	Triangles			Index buffer indicating which vertices make up each triangle. Length must be a multiple of 3.
+	 *	@param	Normals				Optional array of normal vectors for each vertex. If supplied, must be same length as Vertices array.
+	 *	@param	UV0					Optional array of texture co-ordinates for each vertex. If supplied, must be same length as Vertices array.
+	 *	@param	VertexColors		Optional array of colors for each vertex. If supplied, must be same length as Vertices array.
+	 *	@param	Tangents			Optional array of tangent vector for each vertex. If supplied, must be same length as Vertices array.
+	 *	@param	bCreateCollision	Indicates whether collision should be created for this section. This adds significant cost.
+	 */
+	UFUNCTION(BlueprintCallable, Category = "Components|ProceduralMesh", meta = (AutoCreateRefTerm = "Normals,UV0,VertexColors,Tangents"))
+	void CreateMeshSection(int32 SectionIndex, const TArray<FVector>& Vertices, const TArray<int32>& Triangles, const TArray<FVector>& Normals, const TArray<FVector2D>& UV0, const TArray<FColor>& VertexColors, const TArray<FProcMeshTangent>& Tangents, bool bCreateCollision);
+
+	/**
+	 *	Updates a section of this procedural mesh component. This is faster than CreateMeshSection, but does not let you change topology. Collision info is also updated.
+	 *	@param	Vertices			Vertex buffer of all vertex positions to use for this mesh section.
+	 *	@param	Normals				Optional array of normal vectors for each vertex. If supplied, must be same length as Vertices array.
+	 *	@param	UV0					Optional array of texture co-ordinates for each vertex. If supplied, must be same length as Vertices array.
+	 *	@param	VertexColors		Optional array of colors for each vertex. If supplied, must be same length as Vertices array.
+	 *	@param	Tangents			Optional array of tangent vector for each vertex. If supplied, must be same length as Vertices array.
+	 */
+	UFUNCTION(BlueprintCallable, Category = "Components|ProceduralMesh", meta = (AutoCreateRefTerm = "Normals,UV0,VertexColors,Tangents"))
+	void UpdateMeshSection(int32 SectionIndex, const TArray<FVector>& Vertices, const TArray<FVector>& Normals, const TArray<FVector2D>& UV0, const TArray<FColor>& VertexColors, const TArray<FProcMeshTangent>& Tangents);
 
 	/** Clear a section of the procedural mesh. Other sections do not change index. */
 	UFUNCTION(BlueprintCallable, Category = "Components|ProceduralMesh")
-		void ClearMeshSection(int32 SectionIndex);
+	void ClearMeshSection(int32 SectionIndex);
 
 	/** Clear all mesh sections and reset to empty state */
 	UFUNCTION(BlueprintCallable, Category = "Components|ProceduralMesh")
-		void ClearAllMeshSections();
+	void ClearAllMeshSections();
+
+	/** Control visibility of a particular section */
+	UFUNCTION(BlueprintCallable, Category = "Components|ProceduralMesh")
+	void SetMeshSectionVisible(int32 SectionIndex, bool bNewVisibility);
 
+	/** Returns whether a particular section is currently visible */
+	UFUNCTION(BlueprintCallable, Category = "Components|ProceduralMesh")
+	bool IsMeshSectionVisible(int32 SectionIndex) const;
 
 	// Begin Interface_CollisionDataProvider Interface
 	virtual bool GetPhysicsTriMeshData(struct FTriMeshCollisionData* CollisionData, bool InUseAllTriData) override;
@@ -145,6 +176,10 @@ class PROCEDURALMESHCOMPONENT_API UProceduralMeshComponent : public UMeshCompone
 	class UBodySetup* ProcMeshBodySetup;
 
 private:
+	// Begin USceneComponent interface.
+	virtual FBoxSphereBounds CalcBounds(const FTransform& LocalToWorld) const override;
+	// Begin USceneComponent interface.
+
 	// Begin UPrimitiveComponent interface.
 	virtual FPrimitiveSceneProxy* CreateSceneProxy() override;
 	virtual class UBodySetup* GetBodySetup() override;
@@ -154,9 +189,7 @@ private:
 	virtual int32 GetNumMaterials() const override;
 	// End UMeshComponent interface.
 
-	// Begin USceneComponent interface.
-	virtual FBoxSphereBounds CalcBounds(const FTransform& LocalToWorld) const override;
-	// Begin USceneComponent interface.
+
 
 	/** Update LocalBounds member from the local box of each section */
 	void UpdateLocalBounds();
